Adding a program that can be used to delete the PersistentDownloadDir on Windows at uninstall time.

diff --git a/Engine/Source/Programs/AutomationTool/AutomationUtils/MCPPublic.cs b/Engine/Source/Programs/AutomationTool/AutomationUtils/MCPPublic.cs
--- a/Engine/Source/Programs/AutomationTool/AutomationUtils/MCPPublic.cs
+++ b/Engine/Source/Programs/AutomationTool/AutomationUtils/MCPPublic.cs
@@ -647,8 +647,9 @@ namespace EpicGames.MCP.Automation
 			Online_v150,
 			Online_v151,
 			Online_v152,
+			Online_v160,
 
-			Online_Live = Online_v152
+			Online_Live = Online_v160
 		}
 
 		/// <summary>
@@ -750,6 +751,14 @@ namespace EpicGames.MCP.Automation
 			/// </summary>
 			public string AppLaunchCmdArgs;
 			/// <summary>
+			/// The custom uninstall action executable to launch right before product uninstall, relative to the build root.
+			/// </summary>
+			public string UninstallActionPath;
+			/// <summary>
+			/// The commandline to send to the custom uninstall action on launch.
+			/// </summary>
+			public string UninstallActionArgs;
+			/// <summary>
 			/// The list of prerequisite Ids that this prerequisite installer satisfies.
 			/// </summary>
 			public List<string> PrereqIds;
@@ -876,6 +885,14 @@ namespace EpicGames.MCP.Automation
 			/// </summary>
 			public string AppLaunchCmdArgs;
 			/// <summary>
+			/// The custom uninstall action executable to launch right before product uninstall, relative to the build root.
+			/// </summary>
+			public string UninstallActionPath;
+			/// <summary>
+			/// The commandline to send to the custom uninstall action on launch.
+			/// </summary>
+			public string UninstallActionArgs;
+			/// <summary>
 			/// The list of prerequisite Ids that this prerequisite installer satisfies.
 			/// </summary>
 			public List<string> PrereqIds;
diff --git a/Engine/Source/Programs/UninstallHelper/Private/UninstallHelper.cpp b/Engine/Source/Programs/UninstallHelper/Private/UninstallHelper.cpp
new file mode 100644
--- /dev/null
+++ b/Engine/Source/Programs/UninstallHelper/Private/UninstallHelper.cpp
@@ -0,0 +1,222 @@
+// Copyright Epic Games, Inc. All Rights Reserved.
+
+#include "UninstallHelper.h"
+
+#include "RequiredProgramMainCPPInclude.h"
+
+IMPLEMENT_APPLICATION(UninstallHelper, "UninstallHelper");
+DEFINE_LOG_CATEGORY(LogUninstallHelper);
+
+namespace UninstallHelper
+{
+
+	struct FExecActionInfo
+	{
+		FString URL;
+		FString Params;
+	};
+
+	struct FUninstallOptions
+	{
+		bool bShouldExecute = true;
+		FString AppName = "";
+		FString InstallDir = "";
+		bool bDeletePersistentDownloadDir = true;
+		bool bDeleteConfig = false;
+		TArray<FString> AdditionalDeleteArtifacts = {};
+		TArray<FExecActionInfo> AdditionalActions = {};
+	};
+
+	inline FString GetUserSavedDir(const FUninstallOptions& Options)
+	{
+		return FPaths::Combine(FPlatformProcess::UserSettingsDir(), Options.AppName, "Saved");
+	}
+	
+	inline FString GetPersistentDownloadDir(const FUninstallOptions& Options)
+	{
+		return FPaths::Combine(GetUserSavedDir(Options), "PersistentDownloadDir", "");
+	}
+	
+	inline FString GetConfigDir(const FUninstallOptions& Options)
+	{
+		return FPaths::Combine(GetUserSavedDir(Options), "Config", "");
+	}
+
+	FString ReplacePathTokens(const FString& Path, const FUninstallOptions& Options)
+	{
+		FString Result = Path;
+		Result.ReplaceInline(TEXT("%LocalDataDir%"), FPlatformProcess::UserSettingsDir());
+		Result.ReplaceInline(TEXT("%LocalSavedDir%"), *GetUserSavedDir(Options));
+		Result.ReplaceInline(TEXT("%InstallDir%"), *Options.InstallDir);
+		Result.ReplaceInline(TEXT("%GameName%"), *Options.AppName);
+		return Result;
+	}
+
+	int32 ExecProcess(const FExecActionInfo& Action)
+	{
+		int32 Result;
+		FPlatformProcess::ExecElevatedProcess(*Action.URL, *Action.Params, &Result);
+		if (Result)
+		{
+			UE_LOG(LogUninstallHelper, Error, TEXT("Action %s %s failed with exit code %d"), *Action.URL, *Action.Params, Result);
+		}
+		return Result;
+	}
+
+	/**
+	 * Runs the requested uninstall actions.
+	 * @param Options The options to apply to this run.
+	 * @return Whether or not all of the attempted actions were successful.
+	 */
+	EReturnCode Execute(const FUninstallOptions& Options)
+	{
+		// TODO: Decide whether an error should cause a graceful or hard fail
+		EReturnCode Result = EReturnCode::Success;
+		IFileManager& FileManager = IFileManager::Get();
+		if (Options.bDeletePersistentDownloadDir)
+		{
+			const FString PersistentDownloadDir = GetPersistentDownloadDir(Options);
+			UE_LOG(LogUninstallHelper, Log, TEXT("Deleting the PersistentDownloadDir located at %s"), *PersistentDownloadDir);
+			if (Options.bShouldExecute && FileManager.DirectoryExists(*PersistentDownloadDir) && !FileManager.DeleteDirectory(*PersistentDownloadDir, true, true))
+			{
+				Result = EReturnCode::DiskOperationFailed;
+			}
+		}
+		if (Options.bDeleteConfig)
+		{
+			const FString ConfigDir = GetConfigDir(Options);
+			UE_LOG(LogUninstallHelper, Log, TEXT("Deleting the Config directory located at %s"), *ConfigDir);
+			if (Options.bShouldExecute && FileManager.DirectoryExists(*ConfigDir) && !FileManager.DeleteDirectory(*ConfigDir, true, true))
+			{
+				Result = EReturnCode::DiskOperationFailed;
+			}
+		}
+		// Delete artifacts
+		for (const FString& Artifact : Options.AdditionalDeleteArtifacts)
+		{
+			FString CleanArtifactPath = ReplacePathTokens(Artifact, Options);
+			if (FileManager.DirectoryExists(*CleanArtifactPath))
+			{
+				UE_LOG(LogUninstallHelper, Log, TEXT("Deleting directory %s"), *CleanArtifactPath);
+				if (Options.bShouldExecute && !FileManager.DeleteDirectory(*CleanArtifactPath, true, true))
+				{
+					Result = EReturnCode::DiskOperationFailed;
+				}
+			}
+			else if (FileManager.FileExists(*CleanArtifactPath))
+			{
+				UE_LOG(LogUninstallHelper, Log, TEXT("Deleting file %s"), *CleanArtifactPath);
+				if (Options.bShouldExecute && FileManager.Delete(*CleanArtifactPath, true, true))
+				{
+					Result = EReturnCode::DiskOperationFailed;
+				}
+			}
+			else
+			{
+				UE_LOG(LogUninstallHelper, Log, TEXT("Unable to find artifact '%s' on the disk, skipping."), *CleanArtifactPath);
+			}
+		}
+		// Execute other actions
+		for (const FExecActionInfo& Action : Options.AdditionalActions)
+		{
+			UE_LOG(LogUninstallHelper, Log, TEXT("Executing process %s %s..."), *Action.URL, *Action.Params);
+			if (Options.bShouldExecute && ExecProcess(Action) != 0)
+			{
+				Result = EReturnCode::ExecActionFailed;
+			}
+		}
+
+		if (Result != EReturnCode::Success)
+		{
+			UE_LOG(LogUninstallHelper, Error, TEXT("UninstallHelper exited with code %d"), static_cast<int32>(Result));
+		}
+		
+		return Result;
+	}
+	
+}
+
+INT32_MAIN_INT32_ARGC_TCHAR_ARGV()
+{
+	const FString CmdLine = FCommandLine::BuildFromArgV(nullptr, ArgC, ArgV, nullptr);
+	if (!FCommandLine::Set(*CmdLine))
+	{
+		UE_LOG(LogUninstallHelper, Error, TEXT("Failed to initialize command line. Exiting."));
+		return static_cast<int32>(UninstallHelper::EReturnCode::ArgumentError);
+	}
+	
+	UninstallHelper::FUninstallOptions Options;
+
+	if (!FParse::Value(*CmdLine, TEXT("GameName="), Options.AppName) || !FPaths::DirectoryExists(UninstallHelper::GetUserSavedDir(Options)))
+	{
+		UE_LOG(LogUninstallHelper, Error, TEXT("A valid app name must be provided in order for this tool to run. Exiting."));
+		return static_cast<int32>(UninstallHelper::EReturnCode::UnknownAppName);
+	}
+	if (!FParse::Value(*CmdLine, TEXT("InstallDir="), Options.InstallDir) || !FPaths::DirectoryExists(Options.InstallDir))
+	{
+		UE_LOG(LogUninstallHelper, Error, TEXT("A valid install directory must be provided in order for this tool to run. Exiting."));
+		return static_cast<int32>(UninstallHelper::EReturnCode::InvalidInstallDir);
+	}
+	FParse::Bool(*CmdLine, TEXT("DeletePersistentDownloadDir="), Options.bDeletePersistentDownloadDir);
+	FParse::Bool(*CmdLine, TEXT("DeleteConfig="), Options.bDeleteConfig);
+	Options.bShouldExecute = !FParse::Param(*CmdLine, TEXT("NoExecute")); // An option to just log with this application, do not actually do anything
+
+	FString ListDelimiter = ";";
+	FParse::Value(*CmdLine, TEXT("ArtifactDelimiter="), ListDelimiter);
+	
+	FString DeleteDirsCmd;
+	if (FParse::Value(*CmdLine, TEXT("DeleteArtifacts="), DeleteDirsCmd))
+	{
+		DeleteDirsCmd.ParseIntoArray(Options.AdditionalDeleteArtifacts, *ListDelimiter);
+	}
+
+	FString UninstallArtifactsManifestFilename;
+	if (FParse::Value(*CmdLine, TEXT("UninstallArtifactsManifest="), UninstallArtifactsManifestFilename))
+	{
+		TArray<FString> UninstallArtifactsManifest;
+		const FString ManifestPath = UninstallHelper::ReplacePathTokens(UninstallArtifactsManifestFilename, Options);
+		if (!FFileHelper::LoadFileToStringArray(UninstallArtifactsManifest, *ManifestPath))
+		{
+			UE_LOG(LogUninstallHelper, Warning, TEXT("Failed to find uninstall artifacts manifest file '%s', skipping."), *ManifestPath);
+		}
+		Options.AdditionalDeleteArtifacts.Append(UninstallArtifactsManifest);
+	}
+
+	FString UninstallActionsManifestFilename;
+	if (FParse::Value(*CmdLine, TEXT("UninstallActionsManifest="), UninstallActionsManifestFilename))
+	{
+		TArray<FString> UninstallActionsManifest;
+		const FString ManifestPath = UninstallHelper::ReplacePathTokens(UninstallActionsManifestFilename, Options);
+		if (FFileHelper::LoadFileToStringArray(UninstallActionsManifest, *ManifestPath))
+		{
+			for (const FString& Action : UninstallActionsManifest)
+			{
+				FString Url;
+				FString Args;
+				if (Action.StartsWith(TEXT("\"")))
+				{
+					int Size;
+					FParse::QuotedString(*Action, Url, &Size);
+					Url = Url.TrimQuotes();
+					Args = Action.RightChop(Size).TrimStart();
+				}
+				else
+				{
+					Action.Split(TEXT(" "), &Url, &Args);
+				}
+				Options.AdditionalActions.Add({ Url, Args });
+			}
+		}
+		else
+		{
+			UE_LOG(LogUninstallHelper, Warning, TEXT("Failed to find uninstall actions manifest file '%s', skipping."), *ManifestPath);
+		}
+	}
+
+	const UninstallHelper::EReturnCode Result = UninstallHelper::Execute(Options);
+	if (Result != UninstallHelper::EReturnCode::Success)
+	{
+		UE_LOG(LogUninstallHelper, Error, TEXT("UninstallHelper failed to execute with exit code %d"), static_cast<int32>(Result));
+	}
+	return static_cast<int32>(Result);
+}
diff --git a/Engine/Source/Programs/UninstallHelper/Private/UninstallHelper.h b/Engine/Source/Programs/UninstallHelper/Private/UninstallHelper.h
new file mode 100644
--- /dev/null
+++ b/Engine/Source/Programs/UninstallHelper/Private/UninstallHelper.h
@@ -0,0 +1,23 @@
+// Copyright Epic Games, Inc. All Rights Reserved.
+
+#pragma once
+
+#include "CoreMinimal.h"
+
+DECLARE_LOG_CATEGORY_EXTERN(LogUninstallHelper, Log, All);
+
+namespace UninstallHelper
+{
+	enum class EReturnCode : int32
+	{
+		Success = 0,
+		UnknownError,
+		ArgumentError,
+		UnknownAppName,
+		InvalidInstallDir,
+		ExecActionFailed,
+		DiskOperationFailed,
+
+		Crash = 255
+	};
+}
\ No newline at end of file
diff --git a/Engine/Source/Programs/UninstallHelper/UninstallHelper.Build.cs b/Engine/Source/Programs/UninstallHelper/UninstallHelper.Build.cs
new file mode 100644
--- /dev/null
+++ b/Engine/Source/Programs/UninstallHelper/UninstallHelper.Build.cs
@@ -0,0 +1,17 @@
+// Copyright Epic Games, Inc. All Rights Reserved.
+
+using UnrealBuildTool;
+
+public class UninstallHelper : ModuleRules
+{
+	public UninstallHelper(ReadOnlyTargetRules Target) : base(Target)
+	{
+		PublicIncludePaths.Add("Runtime/Launch/Public");
+		PrivateIncludePaths.Add("Runtime/Launch/Private");
+		
+		PrivateDependencyModuleNames.AddRange(new []
+		{
+			"Core", "Projects"
+		});
+	}
+}
diff --git a/Engine/Source/Programs/UninstallHelper/UninstallHelper.Target.cs b/Engine/Source/Programs/UninstallHelper/UninstallHelper.Target.cs
new file mode 100644
--- /dev/null
+++ b/Engine/Source/Programs/UninstallHelper/UninstallHelper.Target.cs
@@ -0,0 +1,34 @@
+﻿// Copyright Epic Games, Inc. All Rights Reserved.
+
+using UnrealBuildTool;
+using System.Collections.Generic;
+
+[SupportedPlatforms(UnrealPlatformClass.Desktop)]
+[SupportedConfigurations(UnrealTargetConfiguration.Debug, UnrealTargetConfiguration.Development, UnrealTargetConfiguration.Shipping)]
+public class UninstallHelperTarget : TargetRules
+{
+	public UninstallHelperTarget(TargetInfo Target) : base(Target)
+	{
+		Type = TargetType.Program;
+		IncludeOrderVersion = EngineIncludeOrderVersion.Latest;
+		LinkType = TargetLinkType.Monolithic;
+		LaunchModuleName = "UninstallHelper";
+
+		// Lean and mean
+		bBuildDeveloperTools = false;
+
+		// Never use malloc profiling in Unreal Header Tool.  We set this because often UHT is compiled right before the engine
+		// automatically by Unreal Build Tool, but if bUseMallocProfiler is defined, UHT can operate incorrectly.
+		bUseMallocProfiler = false;
+		
+		// Currently this app is not linking against the engine, so we'll compile out references from Core to the rest of the engine
+		bCompileAgainstEngine = false;
+		bCompileAgainstCoreUObject = false;
+		bCompileAgainstApplicationCore = false;
+		bBuildWithEditorOnlyData =false;
+		bCompileICU = false;
+
+		// UninstallHelper is a console application, not a Windows app (sets entry point to main(), instead of WinMain())
+		bIsBuildingConsoleApplication = true;
+	}
+}
