Fixing an issue with the UninstallHelper that could cause users who are missing their Saved or PersistentDownloadDir to fail the uninstall process.

diff --git a/Engine/Source/Programs/UninstallHelper/Private/UninstallHelper.cpp b/Engine/Source/Programs/UninstallHelper/Private/UninstallHelper.cpp
--- a/Engine/Source/Programs/UninstallHelper/Private/UninstallHelper.cpp
+++ b/Engine/Source/Programs/UninstallHelper/Private/UninstallHelper.cpp
@@ -19,6 +19,7 @@ namespace UninstallHelper
 	struct FUninstallOptions
 	{
 		bool bShouldExecute = true;
+		bool bAllowHarshReturnCodes = false;
 		FString AppName = "";
 		FString InstallDir = "";
 		bool bDeletePersistentDownloadDir = true;
@@ -76,19 +77,25 @@ namespace UninstallHelper
 		if (Options.bDeletePersistentDownloadDir)
 		{
 			const FString PersistentDownloadDir = GetPersistentDownloadDir(Options);
-			UE_LOG(LogUninstallHelper, Log, TEXT("Deleting the PersistentDownloadDir located at %s"), *PersistentDownloadDir);
-			if (Options.bShouldExecute && FileManager.DirectoryExists(*PersistentDownloadDir) && !FileManager.DeleteDirectory(*PersistentDownloadDir, true, true))
+			if (FileManager.DirectoryExists(*PersistentDownloadDir))
 			{
-				Result = EReturnCode::DiskOperationFailed;
+				UE_LOG(LogUninstallHelper, Log, TEXT("Deleting the PersistentDownloadDir located at %s"), *PersistentDownloadDir);
+				if (Options.bShouldExecute && !FileManager.DeleteDirectory(*PersistentDownloadDir, true, true))
+				{
+					Result = EReturnCode::DiskOperationFailed;
+				}
 			}
 		}
 		if (Options.bDeleteConfig)
 		{
 			const FString ConfigDir = GetConfigDir(Options);
-			UE_LOG(LogUninstallHelper, Log, TEXT("Deleting the Config directory located at %s"), *ConfigDir);
-			if (Options.bShouldExecute && FileManager.DirectoryExists(*ConfigDir) && !FileManager.DeleteDirectory(*ConfigDir, true, true))
+			if (FileManager.DirectoryExists(*ConfigDir))
 			{
-				Result = EReturnCode::DiskOperationFailed;
+				UE_LOG(LogUninstallHelper, Log, TEXT("Deleting the Config directory located at %s"), *ConfigDir);
+				if (Options.bShouldExecute && !FileManager.DeleteDirectory(*ConfigDir, true, true))
+				{
+					Result = EReturnCode::DiskOperationFailed;
+				}
 			}
 		}
 		// Delete artifacts
@@ -106,7 +113,7 @@ namespace UninstallHelper
 			else if (FileManager.FileExists(*CleanArtifactPath))
 			{
 				UE_LOG(LogUninstallHelper, Log, TEXT("Deleting file %s"), *CleanArtifactPath);
-				if (Options.bShouldExecute && FileManager.Delete(*CleanArtifactPath, true, true))
+				if (Options.bShouldExecute && !FileManager.Delete(*CleanArtifactPath, true, true))
 				{
 					Result = EReturnCode::DiskOperationFailed;
 				}
@@ -117,11 +124,13 @@ namespace UninstallHelper
 			}
 		}
 		// Execute other actions
+		int32 LastProcessResult = 0;
 		for (const FExecActionInfo& Action : Options.AdditionalActions)
 		{
 			UE_LOG(LogUninstallHelper, Log, TEXT("Executing process %s %s..."), *Action.URL, *Action.Params);
-			if (Options.bShouldExecute && ExecProcess(Action) != 0)
+			if (Options.bShouldExecute && (LastProcessResult = ExecProcess(Action)) != 0)
 			{
+				UE_LOG(LogUninstallHelper, Log, TEXT("Process %s %s failed with exit code %d"), *Action.URL, *Action.Params, LastProcessResult);
 				Result = EReturnCode::ExecActionFailed;
 			}
 		}
@@ -147,9 +156,9 @@ INT32_MAIN_INT32_ARGC_TCHAR_ARGV()
 	
 	UninstallHelper::FUninstallOptions Options;
 
-	if (!FParse::Value(*CmdLine, TEXT("GameName="), Options.AppName) || !FPaths::DirectoryExists(UninstallHelper::GetUserSavedDir(Options)))
+	if (!FParse::Value(*CmdLine, TEXT("GameName="), Options.AppName))
 	{
-		UE_LOG(LogUninstallHelper, Error, TEXT("A valid app name must be provided in order for this tool to run. Exiting."));
+		UE_LOG(LogUninstallHelper, Error, TEXT("An app name must be provided in order for this tool to run. Exiting."));
 		return static_cast<int32>(UninstallHelper::EReturnCode::UnknownAppName);
 	}
 	if (!FParse::Value(*CmdLine, TEXT("InstallDir="), Options.InstallDir) || !FPaths::DirectoryExists(Options.InstallDir))
@@ -157,6 +166,7 @@ INT32_MAIN_INT32_ARGC_TCHAR_ARGV()
 		UE_LOG(LogUninstallHelper, Error, TEXT("A valid install directory must be provided in order for this tool to run. Exiting."));
 		return static_cast<int32>(UninstallHelper::EReturnCode::InvalidInstallDir);
 	}
+	FParse::Bool(*CmdLine, TEXT("AllowHarshReturnCodes="), Options.bAllowHarshReturnCodes);
 	FParse::Bool(*CmdLine, TEXT("DeletePersistentDownloadDir="), Options.bDeletePersistentDownloadDir);
 	FParse::Bool(*CmdLine, TEXT("DeleteConfig="), Options.bDeleteConfig);
 	Options.bShouldExecute = !FParse::Param(*CmdLine, TEXT("NoExecute")); // An option to just log with this application, do not actually do anything
@@ -218,5 +228,5 @@ INT32_MAIN_INT32_ARGC_TCHAR_ARGV()
 	{
 		UE_LOG(LogUninstallHelper, Error, TEXT("UninstallHelper failed to execute with exit code %d"), static_cast<int32>(Result));
 	}
-	return static_cast<int32>(Result);
+	return Options.bAllowHarshReturnCodes ? static_cast<int32>(Result) : 0;
 }
