Improve UninstallHelper relative path handling and token replacement behavior.

#rnx

diff --git a/Engine/Source/Programs/UninstallHelper/Private/UninstallHelper.cpp b/Engine/Source/Programs/UninstallHelper/Private/UninstallHelper.cpp
--- a/Engine/Source/Programs/UninstallHelper/Private/UninstallHelper.cpp
+++ b/Engine/Source/Programs/UninstallHelper/Private/UninstallHelper.cpp
@@ -24,6 +24,7 @@ namespace UninstallHelper
 		FString InstallDir = "";
 		bool bDeletePersistentDownloadDir = true;
 		bool bDeleteConfig = false;
+		bool bInstallDirRelativeToExecutableDir = false;
 		TArray<FString> AdditionalDeleteArtifacts = {};
 		TArray<FExecActionInfo> AdditionalActions = {};
 	};
@@ -161,9 +162,24 @@ INT32_MAIN_INT32_ARGC_TCHAR_ARGV()
 		UE_LOG(LogUninstallHelper, Error, TEXT("An app name must be provided in order for this tool to run. Exiting."));
 		return static_cast<int32>(UninstallHelper::EReturnCode::UnknownAppName);
 	}
-	if (!FParse::Value(*CmdLine, TEXT("InstallDir="), Options.InstallDir) || !FPaths::DirectoryExists(Options.InstallDir))
+
+	if (!FParse::Value(*CmdLine, TEXT("InstallDir="), Options.InstallDir))
+	{
+		UE_LOG(LogUninstallHelper, Error, TEXT("InstallDir must be provided in order for this tool to run. Exiting."));
+		return static_cast<int32>(UninstallHelper::EReturnCode::InvalidInstallDir);
+	}
+
+	// InstallDir can optionally be treated as relative to the UninstallHelper executable directory instead of the current working directory.
+	FParse::Bool(*CmdLine, TEXT("InstallDirRelativeToExecutableDir="), Options.bInstallDirRelativeToExecutableDir);
+	if (Options.bInstallDirRelativeToExecutableDir && FPaths::IsRelative(Options.InstallDir))
+	{
+		const FString ExecutableDir = FPaths::GetPath(FPlatformProcess::ExecutablePath());
+		Options.InstallDir = FPaths::ConvertRelativePathToFull(ExecutableDir, Options.InstallDir);
+	}
+
+	if (!FPaths::DirectoryExists(Options.InstallDir))
 	{
-		UE_LOG(LogUninstallHelper, Error, TEXT("A valid install directory must be provided in order for this tool to run. Exiting."));
+		UE_LOG(LogUninstallHelper, Error, TEXT("InstallDir %s does not exist. Exiting."), *Options.InstallDir);
 		return static_cast<int32>(UninstallHelper::EReturnCode::InvalidInstallDir);
 	}
 	FParse::Bool(*CmdLine, TEXT("AllowHarshReturnCodes="), Options.bAllowHarshReturnCodes);
@@ -214,6 +230,8 @@ INT32_MAIN_INT32_ARGC_TCHAR_ARGV()
 				{
 					Action.Split(TEXT(" "), &Url, &Args);
 				}
+				Url = ReplacePathTokens(Url, Options);
+				Args = ReplacePathTokens(Args, Options);
 				Options.AdditionalActions.Add({ Url, Args });
 			}
 		}
